AnimNext: Add trajectory analysis functions

#rnx

diff --git a/Engine/Plugins/Experimental/MoverAnimNext/Source/MoverAnimNext/Private/Graph/RigUnit_TrajectoryAnalysis.cpp b/Engine/Plugins/Experimental/MoverAnimNext/Source/MoverAnimNext/Private/Graph/RigUnit_TrajectoryAnalysis.cpp
new file mode 100644
--- /dev/null
+++ b/Engine/Plugins/Experimental/MoverAnimNext/Source/MoverAnimNext/Private/Graph/RigUnit_TrajectoryAnalysis.cpp
@@ -0,0 +1,18 @@
+// Copyright Epic Games, Inc. All Rights Reserved.
+
+#include "RigUnit_TrajectoryAnalysis.h"
+
+FRigUnit_GetTrajectorySampleAtTime_Execute()
+{
+	UPoseSearchTrajectoryLibrary::GetTrajectorySampleAtTime(InTrajectory, Time, OutTrajectorySample, bExtrapolate);
+}
+
+FRigUnit_GetTrajectoryVelocity_Execute()
+{
+	UPoseSearchTrajectoryLibrary::GetTrajectoryVelocity(InTrajectory, Time1, Time2, OutVelocity, bExtrapolate);
+}
+
+FRigUnit_GetTrajectoryAngularVelocity_Execute()
+{
+	UPoseSearchTrajectoryLibrary::GetTrajectoryAngularVelocity(InTrajectory, Time1, Time2, OutAngularVelocity, bExtrapolate);
+}
\ No newline at end of file
diff --git a/Engine/Plugins/Experimental/MoverAnimNext/Source/MoverAnimNext/Private/Graph/RigUnit_TrajectoryAnalysis.h b/Engine/Plugins/Experimental/MoverAnimNext/Source/MoverAnimNext/Private/Graph/RigUnit_TrajectoryAnalysis.h
new file mode 100644
--- /dev/null
+++ b/Engine/Plugins/Experimental/MoverAnimNext/Source/MoverAnimNext/Private/Graph/RigUnit_TrajectoryAnalysis.h
@@ -0,0 +1,78 @@
+// Copyright Epic Games, Inc. All Rights Reserved.
+
+#pragma once
+
+#include "AnimNextExecuteContext.h"
+#include "Graph/RigUnit_AnimNextBase.h"
+#include "PoseSearch/PoseSearchTrajectoryLibrary.h"
+
+#include "RigUnit_TrajectoryAnalysis.generated.h"
+
+USTRUCT(meta = (DisplayName = "Get Trajectory Sample At Time", Category = "Trajectory", NodeColor = "0, 1, 1", Keywords = "Trajectory"))
+struct FRigUnit_GetTrajectorySampleAtTime : public FRigUnit_AnimNextBase
+{
+	GENERATED_BODY()
+
+	RIGVM_METHOD()
+	void Execute();
+
+	UPROPERTY(EditAnywhere, Category = "Trajectory", meta = (Output))
+	FPoseSearchQueryTrajectorySample OutTrajectorySample;
+
+	UPROPERTY(EditAnywhere, Category = "Trajectory", meta = (Input))
+	FPoseSearchQueryTrajectory InTrajectory;
+
+	UPROPERTY(EditAnywhere, Category = "Trajectory", meta = (Input))
+	float Time = 0.0f;
+
+	UPROPERTY(EditAnywhere, Category = "Trajectory", meta = (Input))
+	bool bExtrapolate = false;
+};
+
+USTRUCT(meta = (DisplayName = "Get Trajectory Velocity", Category = "Trajectory", NodeColor = "0, 1, 1", Keywords = "Trajectory"))
+struct FRigUnit_GetTrajectoryVelocity : public FRigUnit_AnimNextBase
+{
+	GENERATED_BODY()
+
+	RIGVM_METHOD()
+	void Execute();
+
+	UPROPERTY(EditAnywhere, Category = "Trajectory", meta = (Input))
+	FPoseSearchQueryTrajectory InTrajectory;
+
+	UPROPERTY(EditAnywhere, Category = "Trajectory", meta = (Output))
+	FVector OutVelocity;
+
+	UPROPERTY(EditAnywhere, Category = "Trajectory", meta = (Input))
+	float Time1 = 0.0f;
+
+	UPROPERTY(EditAnywhere, Category = "Trajectory", meta = (Input))
+	float Time2 = 0.0f;
+
+	UPROPERTY(EditAnywhere, Category = "Trajectory", meta = (Input))
+	bool bExtrapolate = false;
+};
+
+USTRUCT(meta = (DisplayName = "Get Trajectory Angular Velocity", Category = "Trajectory", NodeColor = "0, 1, 1", Keywords = "Trajectory"))
+struct FRigUnit_GetTrajectoryAngularVelocity : public FRigUnit_AnimNextBase
+{
+	GENERATED_BODY()
+
+	RIGVM_METHOD()
+	void Execute();
+
+	UPROPERTY(EditAnywhere, Category = "Trajectory", meta = (Input))
+	FPoseSearchQueryTrajectory InTrajectory;
+
+	UPROPERTY(EditAnywhere, Category = "Trajectory", meta = (Output))
+	FVector OutAngularVelocity;
+
+	UPROPERTY(EditAnywhere, Category = "Trajectory", meta = (Input))
+	float Time1 = 0.0f;
+
+	UPROPERTY(EditAnywhere, Category = "Trajectory", meta = (Input))
+	float Time2 = 0.0f;
+
+	UPROPERTY(EditAnywhere, Category = "Trajectory", meta = (Input))
+	bool bExtrapolate = false;
+};
\ No newline at end of file
