Add defaults to RigUnit_GenerateMoverTrajectory

#tests animsandbox editor win64
#rnx

diff --git a/Engine/Plugins/Experimental/MoverAnimNext/Source/MoverAnimNext/Private/Graph/RigUnit_GenerateMoverTrajectory.h b/Engine/Plugins/Experimental/MoverAnimNext/Source/MoverAnimNext/Private/Graph/RigUnit_GenerateMoverTrajectory.h
--- a/Engine/Plugins/Experimental/MoverAnimNext/Source/MoverAnimNext/Private/Graph/RigUnit_GenerateMoverTrajectory.h
+++ b/Engine/Plugins/Experimental/MoverAnimNext/Source/MoverAnimNext/Private/Graph/RigUnit_GenerateMoverTrajectory.h
@@ -29,16 +29,16 @@ struct FRigUnit_GenerateMoverTrajectory : public FRigUnit_AnimNextBase
 	float DeltaTime;
 
 	UPROPERTY(EditAnywhere, Category = "Graph", meta = (Input))
-	float HistorySamplingInterval;
+	float HistorySamplingInterval = -1.f;
 
 	UPROPERTY(EditAnywhere, Category = "Graph", meta = (Input))
-	int32 NumHistorySamples;
+	int32 NumHistorySamples = 30;
 
 	UPROPERTY(EditAnywhere, Category = "Graph", meta = (Input))
-	float PredictionSamplingInterval;
+	float PredictionSamplingInterval = 0.1f;
 
 	UPROPERTY(EditAnywhere, Category = "Graph", meta = (Input))
-	int32 NumPredictionSamples;
+	int32 NumPredictionSamples = 15;
 
 	UPROPERTY(EditAnywhere, Category = "Graph", meta = (Input, Output))
 	FPoseSearchQueryTrajectory InOutTrajectory;
