MoverAnimNext: adding module function dependency traits for specifying dependencies between different Mover work phases (for standalone backend only)

diff --git a/Engine/Plugins/Experimental/Mover/Source/Mover/Private/Backends/MoverStandaloneLiaison.cpp b/Engine/Plugins/Experimental/Mover/Source/Mover/Private/Backends/MoverStandaloneLiaison.cpp
--- a/Engine/Plugins/Experimental/Mover/Source/Mover/Private/Backends/MoverStandaloneLiaison.cpp
+++ b/Engine/Plugins/Experimental/Mover/Source/Mover/Private/Backends/MoverStandaloneLiaison.cpp
@@ -94,6 +94,31 @@ void UMoverStandaloneLiaisonComponent::BeginPlay()
 	}
 }
 
+FTickFunction* UMoverStandaloneLiaisonComponent::FindTickFunction(EMoverTickPhase MoverTickPhase)
+{
+	FTickFunction* FoundFunction = nullptr;
+
+	switch (MoverTickPhase)
+	{
+		case EMoverTickPhase::ProduceInput:
+			FoundFunction = &ProduceInputTickFunction;
+			break;
+
+		case EMoverTickPhase::SimulateMovement:
+			FoundFunction = &SimulateMovementTickFunction;
+			break;
+
+		case EMoverTickPhase::ApplyState:
+			FoundFunction = &ApplyStateTickFunction;
+			break;
+			
+		default:
+			break;
+	}
+
+	return FoundFunction;
+}
+
 void UMoverStandaloneLiaisonComponent::RegisterComponentTickFunctions(bool bRegister)
 {
 	Super::RegisterComponentTickFunctions(bRegister);
@@ -109,7 +134,7 @@ void UMoverStandaloneLiaisonComponent::RegisterComponentTickFunctions(bool bRegi
 			//   4) Apply new state to actor/components
 
 			ProduceInputTickFunction.Target = this;
-			
+
 			// Input production should always wait for the controller update, and we will watch for controller changes
 			if (APawn* OwnerPawn = Cast<APawn>(GetOwner()))
 			{
@@ -174,7 +199,7 @@ void UMoverStandaloneLiaisonComponent::TickInputProduction(float DeltaSeconds)
 	FMoverInputCmdContext InputCmd;
 	{
 		SCOPED_NAMED_EVENT(StandaloneMoverBackendLiaison_Tick_ProduceInput, FColor::Yellow);
-		MoverComp->ProduceInput(DeltaTimeMs, OUT &LastProducedInputCmd);
+		MoverComp->ProduceInput(DeltaTimeMs, OUT & LastProducedInputCmd);
 	}
 }
 
diff --git a/Engine/Plugins/Experimental/Mover/Source/Mover/Public/Backends/MoverStandaloneLiaison.h b/Engine/Plugins/Experimental/Mover/Source/Mover/Public/Backends/MoverStandaloneLiaison.h
--- a/Engine/Plugins/Experimental/Mover/Source/Mover/Public/Backends/MoverStandaloneLiaison.h
+++ b/Engine/Plugins/Experimental/Mover/Source/Mover/Public/Backends/MoverStandaloneLiaison.h
@@ -112,6 +112,8 @@ public:
 	virtual void BeginPlay() override;
 	// End UActorComponent interface
 
+	FTickFunction* FindTickFunction(EMoverTickPhase MoverTickPhase);
+
 protected:
 	void TickInputProduction(float DeltaSeconds);
 	void TickMovementSimulation(float DeltaSeconds);
diff --git a/Engine/Plugins/Experimental/MoverAnimNext/Source/MoverAnimNext/Private/Module/RigVMTrait_ModuleEventDependency_MoverComponentTickFunctions.cpp b/Engine/Plugins/Experimental/MoverAnimNext/Source/MoverAnimNext/Private/Module/RigVMTrait_ModuleEventDependency_MoverComponentTickFunctions.cpp
new file mode 100644
--- /dev/null
+++ b/Engine/Plugins/Experimental/MoverAnimNext/Source/MoverAnimNext/Private/Module/RigVMTrait_ModuleEventDependency_MoverComponentTickFunctions.cpp
@@ -0,0 +1,81 @@
+// Copyright Epic Games, Inc. All Rights Reserved.
+
+#include "RigVMTrait_ModuleEventDependency_MoverComponentTickFunctions.h"
+#include "Backends/MoverStandaloneLiaison.h"
+
+
+#if WITH_EDITOR
+FString FRigVMTrait_ModuleEventDependency_MoverComponentTickFunctions::GetDisplayName() const
+{
+	return StaticStruct()->GetDisplayNameText().ToString();
+}
+#endif // WITH_EDITOR
+
+void FRigVMTrait_ModuleEventDependency_MoverComponentTickFunctions::OnAddDependency(const UE::AnimNext::FModuleDependencyContext& InContext) const
+{
+	UActorComponent* AnimNextComponent = Cast<UActorComponent>(InContext.Object);
+	if (AnimNextComponent == nullptr)
+	{
+		return;
+	}
+
+	AActor* Actor = AnimNextComponent->GetOwner();
+	if (Actor == nullptr)
+	{
+		return;
+	}
+
+	UMoverStandaloneLiaisonComponent* MoverStandaloneComponent = Actor->FindComponentByClass<UMoverStandaloneLiaisonComponent>();
+	if (MoverStandaloneComponent == nullptr)
+	{
+		return;
+	}
+
+	if (FTickFunction* MoverDependentTickFunction = MoverStandaloneComponent->FindTickFunction(DependentMoverTickPhase))
+	{
+		// How should Mover tick be ordered, relative to the module function?
+		if (Ordering == EAnimNextModuleEventDependencyOrdering::Before)
+		{
+			InContext.TickFunction.AddPrerequisite(MoverStandaloneComponent, *MoverDependentTickFunction);
+		}
+		else
+		{
+			MoverDependentTickFunction->AddPrerequisite(AnimNextComponent, InContext.TickFunction);
+		}
+	}
+}
+
+
+void FRigVMTrait_ModuleEventDependency_MoverComponentTickFunctions::OnRemoveDependency(const UE::AnimNext::FModuleDependencyContext& InContext) const
+{
+	UActorComponent* AnimNextComponent = Cast<UActorComponent>(InContext.Object);
+	if (AnimNextComponent == nullptr)
+	{
+		return;
+	}
+
+	AActor* Actor = AnimNextComponent->GetOwner();
+	if (Actor == nullptr)
+	{
+		return;
+	}
+
+	UMoverStandaloneLiaisonComponent* MoverStandaloneComponent = Actor->FindComponentByClass<UMoverStandaloneLiaisonComponent>();
+	if (MoverStandaloneComponent == nullptr)
+	{
+		return;
+	}
+
+	if (FTickFunction* MoverDependentTickFunction = MoverStandaloneComponent->FindTickFunction(DependentMoverTickPhase))
+	{
+		// How should Mover tick be ordered, relative to the module function?
+		if (Ordering == EAnimNextModuleEventDependencyOrdering::Before)
+		{
+			InContext.TickFunction.RemovePrerequisite(MoverStandaloneComponent, *MoverDependentTickFunction);
+		}
+		else
+		{
+			MoverDependentTickFunction->RemovePrerequisite(AnimNextComponent, InContext.TickFunction);
+		}
+	}
+}
\ No newline at end of file
diff --git a/Engine/Plugins/Experimental/MoverAnimNext/Source/MoverAnimNext/Private/Module/RigVMTrait_ModuleEventDependency_MoverComponentTickFunctions.h b/Engine/Plugins/Experimental/MoverAnimNext/Source/MoverAnimNext/Private/Module/RigVMTrait_ModuleEventDependency_MoverComponentTickFunctions.h
new file mode 100644
--- /dev/null
+++ b/Engine/Plugins/Experimental/MoverAnimNext/Source/MoverAnimNext/Private/Module/RigVMTrait_ModuleEventDependency_MoverComponentTickFunctions.h
@@ -0,0 +1,32 @@
+// Copyright Epic Games, Inc. All Rights Reserved.
+
+#pragma once
+
+#include "Module/RigVMTrait_ModuleEventDependency.h"
+#include "MoverTypes.h"
+#include "RigVMTrait_ModuleEventDependency_MoverComponentTickFunctions.generated.h"
+
+
+/** A dependency on one of the tick functions that make up a MoverComponent's update flow. This uses 
+ * the first-found MoverComponent on the current actor.
+ * TODO: add support beyond standalone (non-networked) Mover backends
+ */
+USTRUCT(DisplayName = "Mover Component Tick Functions", meta = (ShowTooltip = true))
+struct FRigVMTrait_ModuleEventDependency_MoverComponentTickFunctions : public FRigVMTrait_ModuleEventDependency
+{
+	GENERATED_BODY()
+
+	// FRigVMTrait interface
+#if WITH_EDITOR
+	virtual FString GetDisplayName() const override;
+#endif
+
+	// FAnimNextModuleEventDependency interface
+	virtual void OnAddDependency(const UE::AnimNext::FModuleDependencyContext& InContext) const override;
+	virtual void OnRemoveDependency(const UE::AnimNext::FModuleDependencyContext& InContext) const override;
+
+
+	// The Mover tick phase that this dependency relates to
+	UPROPERTY(EditAnywhere, Category = "Dependency")
+	EMoverTickPhase DependentMoverTickPhase = EMoverTickPhase::ApplyState;
+};
\ No newline at end of file
