Removing unnecessary dependency from MoverAnimNext that was causing build warnings

diff --git a/Engine/Plugins/Experimental/MoverAnimNext/Source/MoverAnimNext/MoverAnimNext.Build.cs b/Engine/Plugins/Experimental/MoverAnimNext/Source/MoverAnimNext/MoverAnimNext.Build.cs
--- a/Engine/Plugins/Experimental/MoverAnimNext/Source/MoverAnimNext/MoverAnimNext.Build.cs
+++ b/Engine/Plugins/Experimental/MoverAnimNext/Source/MoverAnimNext/MoverAnimNext.Build.cs
@@ -23,7 +23,6 @@ public class MoverAnimNext : ModuleRules
 			{
 				"CoreUObject",
 				"RigVM",
-				"ControlRig",
 				"Engine",
 				"PhysicsCore",
 				"DeveloperSettings"
