Trajectory - adding missing default initializers

diff --git a/Engine/Plugins/Experimental/AnimNextPoseSearch/Source/AnimNextPoseSearch/Private/RigUnit_DebugDrawTrajectory.h b/Engine/Plugins/Experimental/AnimNextPoseSearch/Source/AnimNextPoseSearch/Private/RigUnit_DebugDrawTrajectory.h
--- a/Engine/Plugins/Experimental/AnimNextPoseSearch/Source/AnimNextPoseSearch/Private/RigUnit_DebugDrawTrajectory.h
+++ b/Engine/Plugins/Experimental/AnimNextPoseSearch/Source/AnimNextPoseSearch/Private/RigUnit_DebugDrawTrajectory.h
@@ -20,11 +20,11 @@ struct FRigUnit_DebugDrawTrajectory: public FRigUnit_AnimNextBase
 	FPoseSearchQueryTrajectory Trajectory;
 
 	UPROPERTY(EditAnywhere, Category = "Graph", meta = (Input))
-	float DebugThickness;
+	float DebugThickness = 0.f;
 
 	UPROPERTY(EditAnywhere, Category = "Graph", meta = (Input))
-	float DebugOffset;
+	float DebugOffset = 0.f;
 
 	UPROPERTY(EditAnywhere, Category = "Graph", meta = (Input))
-	bool Enabled;
+	bool Enabled = false;
 };
diff --git a/Engine/Plugins/Experimental/MoverAnimNext/Source/MoverAnimNext/Private/Graph/RigUnit_GenerateMoverTrajectory.h b/Engine/Plugins/Experimental/MoverAnimNext/Source/MoverAnimNext/Private/Graph/RigUnit_GenerateMoverTrajectory.h
--- a/Engine/Plugins/Experimental/MoverAnimNext/Source/MoverAnimNext/Private/Graph/RigUnit_GenerateMoverTrajectory.h
+++ b/Engine/Plugins/Experimental/MoverAnimNext/Source/MoverAnimNext/Private/Graph/RigUnit_GenerateMoverTrajectory.h
@@ -23,7 +23,7 @@ struct FRigUnit_GenerateMoverTrajectory : public FRigUnit_AnimNextBase
 
 	UPROPERTY(EditAnywhere, Category = "Graph", meta = (Input))
 	// This should be the most recent simulation time that was used to get us to our current state
-	float DeltaTime;
+	float DeltaTime = 0.f;
 
 	UPROPERTY(EditAnywhere, Category = "Graph", meta = (Input))
 	float HistorySamplingInterval = -1.f;
diff --git a/Engine/Plugins/Experimental/MoverAnimNext/Source/MoverAnimNext/Private/Graph/RigUnit_TrajectoryAnalysis.h b/Engine/Plugins/Experimental/MoverAnimNext/Source/MoverAnimNext/Private/Graph/RigUnit_TrajectoryAnalysis.h
--- a/Engine/Plugins/Experimental/MoverAnimNext/Source/MoverAnimNext/Private/Graph/RigUnit_TrajectoryAnalysis.h
+++ b/Engine/Plugins/Experimental/MoverAnimNext/Source/MoverAnimNext/Private/Graph/RigUnit_TrajectoryAnalysis.h
@@ -41,7 +41,7 @@ struct FRigUnit_GetTrajectoryVelocity : public FRigUnit_AnimNextBase
 	FPoseSearchQueryTrajectory InTrajectory;
 
 	UPROPERTY(EditAnywhere, Category = "Trajectory", meta = (Output))
-	FVector OutVelocity;
+	FVector OutVelocity = FVector::ZeroVector;
 
 	UPROPERTY(EditAnywhere, Category = "Trajectory", meta = (Input))
 	float Time1 = 0.0f;
@@ -65,7 +65,7 @@ struct FRigUnit_GetTrajectoryAngularVelocity : public FRigUnit_AnimNextBase
 	FPoseSearchQueryTrajectory InTrajectory;
 
 	UPROPERTY(EditAnywhere, Category = "Trajectory", meta = (Output))
-	FVector OutAngularVelocity;
+	FVector OutAngularVelocity = FVector::ZeroVector;
 
 	UPROPERTY(EditAnywhere, Category = "Trajectory", meta = (Input))
 	float Time1 = 0.0f;
