Updated UAF Trajectory functions to have an execution pin.

diff --git a/Engine/Plugins/Experimental/MoverAnimNext/Source/MoverAnimNext/Private/Graph/RigUnit_GenerateMoverTrajectory.h b/Engine/Plugins/Experimental/MoverAnimNext/Source/MoverAnimNext/Private/Graph/RigUnit_GenerateMoverTrajectory.h
--- a/Engine/Plugins/Experimental/MoverAnimNext/Source/MoverAnimNext/Private/Graph/RigUnit_GenerateMoverTrajectory.h
+++ b/Engine/Plugins/Experimental/MoverAnimNext/Source/MoverAnimNext/Private/Graph/RigUnit_GenerateMoverTrajectory.h
@@ -18,25 +18,29 @@ struct FRigUnit_GenerateMoverTrajectory : public FRigUnit_AnimNextBase
 	RIGVM_METHOD()
 	void Execute();
 
-	UPROPERTY(EditAnywhere, Category = "Graph", meta = (Input))
+	UPROPERTY(EditAnywhere, Category = "Trajectory", meta = (Input))
 	TObjectPtr<UMoverComponent> MoverComponent;
 
-	UPROPERTY(EditAnywhere, Category = "Graph", meta = (Input))
+	UPROPERTY(EditAnywhere, Category = "Trajectory", meta = (Input))
 	// This should be the most recent simulation time that was used to get us to our current state
 	float DeltaTime = 0.f;
 
-	UPROPERTY(EditAnywhere, Category = "Graph", meta = (Input))
+	UPROPERTY(EditAnywhere, Category = "Trajectory", meta = (Input))
 	float HistorySamplingInterval = -1.f;
 
-	UPROPERTY(EditAnywhere, Category = "Graph", meta = (Input))
+	UPROPERTY(EditAnywhere, Category = "Trajectory", meta = (Input))
 	int32 NumHistorySamples = 30;
 
-	UPROPERTY(EditAnywhere, Category = "Graph", meta = (Input))
+	UPROPERTY(EditAnywhere, Category = "Trajectory", meta = (Input))
 	float PredictionSamplingInterval = 0.1f;
 
-	UPROPERTY(EditAnywhere, Category = "Graph", meta = (Input))
+	UPROPERTY(EditAnywhere, Category = "Trajectory", meta = (Input))
 	int32 NumPredictionSamples = 15;
 
-	UPROPERTY(EditAnywhere, Category = "Graph", meta = (Input, Output))
+	UPROPERTY(EditAnywhere, Category = "Trajectory", meta = (Input, Output))
 	FTransformTrajectory InOutTrajectory;
+	
+	// The execution result
+	UPROPERTY(EditAnywhere, DisplayName = "Execute", Category = "BeginExecution", meta = (Input, Output))
+	FAnimNextExecuteContext ExecuteContext;
 };
diff --git a/Engine/Plugins/Experimental/UAF/UAFPoseSearch/Source/UAFPoseSearch/Private/RigUnit_DebugDrawTrajectory.h b/Engine/Plugins/Experimental/UAF/UAFPoseSearch/Source/UAFPoseSearch/Private/RigUnit_DebugDrawTrajectory.h
--- a/Engine/Plugins/Experimental/UAF/UAFPoseSearch/Source/UAFPoseSearch/Private/RigUnit_DebugDrawTrajectory.h
+++ b/Engine/Plugins/Experimental/UAF/UAFPoseSearch/Source/UAFPoseSearch/Private/RigUnit_DebugDrawTrajectory.h
@@ -8,7 +8,7 @@
 
 #include "RigUnit_DebugDrawTrajectory.generated.h"
 
-USTRUCT(meta = (DisplayName = "Debug Draw Trajectory", Category = "Mover", NodeColor = "0, 1, 1", Keywords = "Mover, Trajectory, Debug"))
+USTRUCT(meta = (DisplayName = "Debug Draw Trajectory", Category = "Trajectory", NodeColor = "0, 1, 1", Keywords = "Mover, Trajectory, Debug"))
 struct FRigUnit_DebugDrawTrajectory: public FRigUnit_AnimNextBase
 {
 	GENERATED_BODY()
@@ -16,15 +16,19 @@ struct FRigUnit_DebugDrawTrajectory: public FRigUnit_AnimNextBase
 	RIGVM_METHOD()
 	void Execute();
 
-	UPROPERTY(EditAnywhere, Category = "Graph", meta = (Input, Output))
+	UPROPERTY(EditAnywhere, Category = "Debug Draw - Trajectory", meta = (Input, Output))
 	FTransformTrajectory Trajectory;
 
-	UPROPERTY(EditAnywhere, Category = "Graph", meta = (Input))
+	UPROPERTY(EditAnywhere, Category = "Debug Draw - Trajectory", meta = (Input))
 	float DebugThickness = 0.f;
 
-	UPROPERTY(EditAnywhere, Category = "Graph", meta = (Input))
+	UPROPERTY(EditAnywhere, Category = "Debug Draw - Trajectory", meta = (Input))
 	float DebugOffset = 0.f;
 
-	UPROPERTY(EditAnywhere, Category = "Graph", meta = (Input))
+	UPROPERTY(EditAnywhere, Category = "Debug Draw - Trajectory", meta = (Input))
 	bool Enabled = false;
+
+	// The execution result
+	UPROPERTY(EditAnywhere, DisplayName = "Execute", Category = "BeginExecution", meta = (Input, Output))
+	FAnimNextExecuteContext ExecuteContext;
 };
