Fix crash in RigUnit_GenerateMoverTrajectory

diff --git a/Engine/Plugins/Experimental/MoverAnimNext/Source/MoverAnimNext/Private/Graph/RigUnit_GenerateMoverTrajectory.cpp b/Engine/Plugins/Experimental/MoverAnimNext/Source/MoverAnimNext/Private/Graph/RigUnit_GenerateMoverTrajectory.cpp
--- a/Engine/Plugins/Experimental/MoverAnimNext/Source/MoverAnimNext/Private/Graph/RigUnit_GenerateMoverTrajectory.cpp
+++ b/Engine/Plugins/Experimental/MoverAnimNext/Source/MoverAnimNext/Private/Graph/RigUnit_GenerateMoverTrajectory.cpp
@@ -38,5 +38,5 @@ FRigUnit_GenerateMoverTrajectory_Execute()
 	UPoseSearchTrajectoryLibrary::InitTrajectorySamples(InOutTrajectory, CurrentPosition, CurrentFacing, TrajectoryDataSampling, DeltaTime);
 	UPoseSearchTrajectoryLibrary::UpdateHistory_TransformHistory(InOutTrajectory, CurrentPosition, CurrentVelocity, TrajectoryDataSampling, DeltaTime);
 	UMoverTrajectoryPredictor::Predict(*MoverComponent, InOutTrajectory,
-		NumPredictionSamples + 1, PredictionSamplingInterval, NumHistorySamples, static_cast<float>(MoverSamplingFrameRate.AsInterval()));
+		NumPredictionSamples, PredictionSamplingInterval, NumHistorySamples, static_cast<float>(MoverSamplingFrameRate.AsInterval()));
 }
