AnimSandbox: Fix system compilation error

Add FFrameRate as a RigVM allowed struct type to support FRigUnit_GenerateMoverTrajectory

#virtualized

diff --git a/Engine/Plugins/Experimental/MoverAnimNext/Source/MoverAnimNext/Private/MoverAnimNextModule.cpp b/Engine/Plugins/Experimental/MoverAnimNext/Source/MoverAnimNext/Private/MoverAnimNextModule.cpp
--- a/Engine/Plugins/Experimental/MoverAnimNext/Source/MoverAnimNext/Private/MoverAnimNextModule.cpp
+++ b/Engine/Plugins/Experimental/MoverAnimNext/Source/MoverAnimNext/Private/MoverAnimNextModule.cpp
@@ -20,7 +20,7 @@ void FMoverAnimNextModule::StartupModule()
 		FTrajectorySampleInfo::StaticStruct(),			// Todo: remove
 		FPoseSearchTrajectoryData::StaticStruct(),
 		FTransformTrajectory::StaticStruct(),
-
+		TBaseStructure<FFrameRate>::Get()
 	};
 
 	static TPair<UClass*, FRigVMRegistry::ERegisterObjectOperation> const AllowedObjectTypes[] =
