Fix for broken trajectories on Mover UAF characters

diff --git a/Engine/Plugins/Experimental/MoverAnimNext/Source/MoverAnimNext/Private/Graph/RigUnit_GenerateMoverTrajectory.cpp b/Engine/Plugins/Experimental/MoverAnimNext/Source/MoverAnimNext/Private/Graph/RigUnit_GenerateMoverTrajectory.cpp
--- a/Engine/Plugins/Experimental/MoverAnimNext/Source/MoverAnimNext/Private/Graph/RigUnit_GenerateMoverTrajectory.cpp
+++ b/Engine/Plugins/Experimental/MoverAnimNext/Source/MoverAnimNext/Private/Graph/RigUnit_GenerateMoverTrajectory.cpp
@@ -37,6 +37,14 @@ FRigUnit_GenerateMoverTrajectory_Execute()
 	UMoverTrajectoryPredictor::GetCurrentState(*MoverComponent, CurrentPosition, CurrentFacing, CurrentVelocity);
 	UPoseSearchTrajectoryLibrary::InitTrajectorySamples(InOutTrajectory, CurrentPosition, CurrentFacing, TrajectoryDataSampling, DeltaTime);
 	UPoseSearchTrajectoryLibrary::UpdateHistory_TransformHistory(InOutTrajectory, CurrentPosition, CurrentVelocity, TrajectoryDataSampling, DeltaTime);
+	
+	// Set the current position at i == NumHistorySamples at t == delta time
+	// Remember: t == 0 is the previous position, t == delta time is the current frame's position
+	// assuming we call this method after the movement component has updated to a new position
+	InOutTrajectory.Samples[NumHistorySamples].TimeInSeconds = DeltaTime;
+	InOutTrajectory.Samples[NumHistorySamples].Position = CurrentPosition;
+	InOutTrajectory.Samples[NumHistorySamples].Facing = CurrentFacing;
+	
 	UMoverTrajectoryPredictor::Predict(*MoverComponent, InOutTrajectory,
 		NumPredictionSamples, PredictionSamplingInterval, NumHistorySamples, static_cast<float>(MoverSamplingFrameRate.AsInterval()));
 }
